Add a small zip file utility.

diff --git a/Engine/Source/Developer/FileUtilities/FileUtilities.Build.cs b/Engine/Source/Developer/FileUtilities/FileUtilities.Build.cs
new file mode 100644
--- /dev/null
+++ b/Engine/Source/Developer/FileUtilities/FileUtilities.Build.cs
@@ -0,0 +1,11 @@
+// Copyright 1998-2018 Epic Games, Inc. All Rights Reserved.
+
+using UnrealBuildTool;
+
+public class FileUtilities : ModuleRules
+{
+	public FileUtilities(ReadOnlyTargetRules Target) : base(Target)
+	{
+		PrivateDependencyModuleNames.AddRange(new string[] { "Core"});
+    }
+}
diff --git a/Engine/Source/Developer/FileUtilities/Private/FileUtilities.cpp b/Engine/Source/Developer/FileUtilities/Private/FileUtilities.cpp
new file mode 100644
--- /dev/null
+++ b/Engine/Source/Developer/FileUtilities/Private/FileUtilities.cpp
@@ -0,0 +1,157 @@
+// Copyright 1998-2018 Epic Games, Inc. All Rights Reserved.
+
+#include "FileUtilities.h"
+
+DECLARE_LOG_CATEGORY_CLASS(LogZipArchiveWriter, Display, Display);
+
+#if WITH_ENGINE
+
+FZipArchiveWriter::FZipArchiveWriter(IFileHandle* InFile)
+	: File(InFile)
+{
+}
+
+FZipArchiveWriter::~FZipArchiveWriter()
+{
+	// Zip File Format Specification:
+	// https://www.loc.gov/preservation/digital/formats/digformatspecs/APPNOTE%2820120901%29_Version_6.3.3.txt
+
+	UE_LOG(LogZipArchiveWriter, Display, TEXT("Closing zip file with %d entries."), Files.Num());
+
+	// Write the file directory
+	uint64 DirStartOffset = Tell();
+	for (FFileEntry& Entry : Files)
+	{
+		const static uint8 Footer[] =
+		{
+			0x50, 0x4b, 0x01, 0x02, 0x3f, 0x00, 0x2d, 0x00,
+			0x00, 0x00, 0x00, 0x00
+		};
+		Write((void*)Footer, sizeof(Footer));
+		Write(Entry.Time);
+		Write(Entry.Crc32);
+
+		Write((uint64)0xffffffffffffffff);
+		Write((uint16)Entry.Filename.Len());
+		const static uint8 Fields[] =
+		{
+			0x20, 0x00, 0x00, 0x00, 0x00, 0x00, 0x00, 0x00,
+			0x20, 0x00, 0x00, 0x00, 0xff, 0xff, 0xff, 0xff
+		};
+		Write((void*)Fields, sizeof(Fields));
+		Write((void*)TCHAR_TO_UTF8(*Entry.Filename), Entry.Filename.Len());
+
+		Write((uint16)0x01);
+		Write((uint16)0x1c);
+
+		Write((uint64)Entry.Length);
+		Write((uint64)Entry.Length);
+		Write((uint64)Entry.Offset);
+		Write((uint32)0);
+
+		Flush();
+	}
+	uint64 DirEndOffset = Tell();
+
+	uint64 DirectorySizeInBytes = DirEndOffset - DirStartOffset;
+
+	// Write ZIP64 end of central directory record
+	const static uint8 Record[] =
+	{
+		0x50, 0x4b, 0x06, 0x06, 0x2c, 0x00, 0x00, 0x00,
+		0x00, 0x00, 0x00, 0x00, 0x2d, 0x00, 0x2d, 0x00,
+		0x00, 0x00, 0x00, 0x00, 0x00, 0x00, 0x00, 0x00,
+	};
+	Write((void*)Record, sizeof(Record));
+	Write((uint64)Files.Num());
+	Write((uint64)Files.Num());
+	Write(DirectorySizeInBytes);
+	Write(DirStartOffset);
+
+	// Write ZIP64 end of central directory locator
+	const static uint8 Locator[] =
+	{
+		0x50, 0x4b, 0x06, 0x07, 0x00, 0x00, 0x00, 0x00,
+	};
+	Write((void*)Locator, sizeof(Locator));
+	Write(DirEndOffset);
+	Write((uint32)0x01);
+
+	// Write normal end of central directory record
+	const static uint8 EndRecord[] =
+	{
+		0x50, 0x4b, 0x05, 0x06, 0xff, 0xff, 0xff, 0xff,
+		0xff, 0xff, 0xff, 0xff, 0xff, 0xff, 0xff, 0xff,
+		0xff, 0xff, 0xff, 0xff, 0x00, 0x00
+	};
+	Write((void*)EndRecord, sizeof(EndRecord));
+
+	Flush();
+
+	if (File)
+	{
+		// Close the file
+		delete File;
+		File = nullptr;
+	}
+}
+
+void FZipArchiveWriter::AddFile(const FString& Filename, const TArray<uint8>& Data, const FDateTime& Timestamp)
+{
+	uint32 Crc = FCrc::MemCrc32(Data.GetData(), Data.Num());
+
+	// Convert the date-time to a zip file timestamp (2-second resolution).
+	uint32 ZipTime =
+		(Timestamp.GetSecond() / 2) |
+		(Timestamp.GetMinute() << 5) |
+		(Timestamp.GetHour() << 11) |
+		(Timestamp.GetDay() << 16) |
+		(Timestamp.GetMonth() << 21) |
+		((Timestamp.GetYear() - 1980) << 25);
+
+	uint64 FileOffset = Tell();
+
+	FFileEntry* Entry = new (Files) FFileEntry(Filename, Crc, Data.Num(), FileOffset, ZipTime);
+
+	static const uint8 Header[] =
+	{
+		0x50, 0x4b, 0x03, 0x04, 0x2d, 0x00, 0x00, 0x00,
+		0x00, 0x00
+	};
+	Write((void*)Header, sizeof(Header));
+	Write(ZipTime);
+	Write(Crc);
+	Write((uint64)0xffffffffffffffff);
+	Write((uint16)Filename.Len());
+	Write((uint16)0x20);
+
+	Write((void*)TCHAR_TO_UTF8(*Entry->Filename), Filename.Len());
+
+	Write((uint16)0x01);
+	Write((uint16)0x1c);
+	Write((uint64)Data.Num());
+	Write((uint64)Data.Num());
+	Write((uint64)FileOffset);
+	Write((uint32)0);
+
+	Write((void*)Data.GetData(), Data.Num());
+
+	Flush();
+}
+
+void FZipArchiveWriter::Flush()
+{
+	if (Buffer.Num())
+	{
+		if (File && !File->Write(Buffer.GetData(), Buffer.Num()))
+		{
+			UE_LOG(LogZipArchiveWriter, Error, TEXT("Failed to write to zip file. Zip file writing aborted."));
+			delete File;
+			File = nullptr;
+		}
+
+		Buffer.Reset(Buffer.Num());
+	}
+}
+
+#endif
\ No newline at end of file
diff --git a/Engine/Source/Developer/FileUtilities/Public/FileUtilities.h b/Engine/Source/Developer/FileUtilities/Public/FileUtilities.h
new file mode 100644
--- /dev/null
+++ b/Engine/Source/Developer/FileUtilities/Public/FileUtilities.h
@@ -0,0 +1,50 @@
+// Copyright 1998-2018 Epic Games, Inc. All Rights Reserved.
+
+#include "CoreMinimal.h"
+
+#if WITH_ENGINE
+
+/** Helper class for generating an uncompressed zip archive file. */
+class FZipArchiveWriter
+{
+	struct FFileEntry
+	{
+		FString Filename;
+		uint32 Crc32;
+		uint64 Length;
+		uint64 Offset;
+		uint32 Time;
+
+		FFileEntry(const FString& InFilename, uint32 InCrc32, uint64 InLength, uint64 InOffset, uint32 InTime)
+			: Filename(InFilename)
+			, Crc32(InCrc32)
+			, Length(InLength)
+			, Offset(InOffset)
+			, Time(InTime)
+		{}
+	};
+
+	TArray<FFileEntry> Files;
+
+	TArray<uint8> Buffer;
+	IFileHandle* File;
+
+	inline void Write(uint16 V) { Write((void*)&V, sizeof(V)); }
+	inline void Write(uint32 V) { Write((void*)&V, sizeof(V)); }
+	inline void Write(uint64 V) { Write((void*)&V, sizeof(V)); }
+	inline void Write(void* Src, uint64 Size)
+	{
+		void* Dst = &Buffer[Buffer.AddUninitialized(Size)];
+		FMemory::Memcpy(Dst, Src, Size);
+	}
+	inline uint64 Tell() { return (File ? File->Tell() : 0) + Buffer.Num(); }
+	void Flush();
+
+public:
+	FZipArchiveWriter(IFileHandle* InFile);
+	~FZipArchiveWriter();
+
+	void AddFile(const FString& Filename, const TArray<uint8>& Data, const FDateTime& Timestamp);
+};
+
+#endif
\ No newline at end of file
