Add missing forward declaration of IFileHandle

diff --git a/Engine/Source/Developer/FileUtilities/Public/FileUtilities/ZipArchiveWriter.h b/Engine/Source/Developer/FileUtilities/Public/FileUtilities/ZipArchiveWriter.h
--- a/Engine/Source/Developer/FileUtilities/Public/FileUtilities/ZipArchiveWriter.h
+++ b/Engine/Source/Developer/FileUtilities/Public/FileUtilities/ZipArchiveWriter.h
@@ -4,6 +4,8 @@
 
 #if WITH_ENGINE
 
+class IFileHandle;
+
 /** Helper class for generating an uncompressed zip archive file. */
 class FILEUTILITIES_API FZipArchiveWriter
 {
