Fix missing module declaration in FileUtilities (4827582 introduced a link error check).

diff --git a/Engine/Source/Developer/FileUtilities/FileUtilities.Build.cs b/Engine/Source/Developer/FileUtilities/FileUtilities.Build.cs
--- a/Engine/Source/Developer/FileUtilities/FileUtilities.Build.cs
+++ b/Engine/Source/Developer/FileUtilities/FileUtilities.Build.cs
@@ -1,4 +1,4 @@
-// Copyright 1998-2018 Epic Games, Inc. All Rights Reserved.
+// Copyright 1998-2019 Epic Games, Inc. All Rights Reserved.
 
 using UnrealBuildTool;
 
diff --git a/Engine/Source/Developer/FileUtilities/Private/FileUtilities.cpp b/Engine/Source/Developer/FileUtilities/Private/FileUtilities.cpp
new file mode 100644
--- /dev/null
+++ b/Engine/Source/Developer/FileUtilities/Private/FileUtilities.cpp
@@ -0,0 +1,5 @@
+// Copyright 1998-2019 Epic Games, Inc. All Rights Reserved.
+
+#include "Modules/ModuleManager.h"
+
+IMPLEMENT_MODULE(FDefaultModuleImpl, FileUtilities);
\ No newline at end of file
