Fix non-unity compile error.

diff --git a/Engine/Source/Developer/FileUtilities/Public/FileUtilities/ZipArchiveWriter.h b/Engine/Source/Developer/FileUtilities/Public/FileUtilities/ZipArchiveWriter.h
--- a/Engine/Source/Developer/FileUtilities/Public/FileUtilities/ZipArchiveWriter.h
+++ b/Engine/Source/Developer/FileUtilities/Public/FileUtilities/ZipArchiveWriter.h
@@ -1,11 +1,10 @@
 // Copyright 1998-2019 Epic Games, Inc. All Rights Reserved.
 
 #include "CoreMinimal.h"
+#include "GenericPlatform/GenericPlatformFile.h"
 
 #if WITH_ENGINE
 
-class IFileHandle;
-
 /** Helper class for generating an uncompressed zip archive file. */
 class FILEUTILITIES_API FZipArchiveWriter
 {
