FileUtilities: Disable ZipArchiveReader on PLATFORM_MAC until we have fixed the libzip.a library.

diff --git a/Engine/Source/Developer/FileUtilities/FileUtilities.Build.cs b/Engine/Source/Developer/FileUtilities/FileUtilities.Build.cs
--- a/Engine/Source/Developer/FileUtilities/FileUtilities.Build.cs
+++ b/Engine/Source/Developer/FileUtilities/FileUtilities.Build.cs
@@ -10,7 +10,11 @@ public class FileUtilities : ModuleRules
 
 		if (Target.bBuildEditor)
 		{
-			PrivateDependencyModuleNames.AddRange(new string[] { "libzip" });
+			// Temporarily disabled while we fix libzip for mac
+			if (Target.Platform != UnrealTargetPlatform.Mac)
+			{
+				PrivateDependencyModuleNames.AddRange(new string[] { "libzip" });
+			}
 		}
 	}
 }
diff --git a/Engine/Source/Developer/FileUtilities/Private/ZipArchiveReader.cpp b/Engine/Source/Developer/FileUtilities/Private/ZipArchiveReader.cpp
--- a/Engine/Source/Developer/FileUtilities/Private/ZipArchiveReader.cpp
+++ b/Engine/Source/Developer/FileUtilities/Private/ZipArchiveReader.cpp
@@ -4,6 +4,46 @@
 
 #if WITH_EDITOR
 
+// Temporarily disabled while we fix libzip for mac
+#if PLATFORM_MAC
+
+#include "ZipArchivePrivate.h"
+
+class FZipArchiveReader::FImpl
+{
+};
+
+FZipArchiveReader::FZipArchiveReader(IFileHandle* InFileHandle, FOutputDevice* ErrorHandler)
+{
+	if (ErrorHandler)
+	{
+		ErrorHandler->Log(LogZipArchive.GetCategoryName(), ELogVerbosity::Display, TEXT("Not yet implemented on PLATFORM_MAC"));
+	}
+}
+
+FZipArchiveReader::~FZipArchiveReader() = default;
+
+bool FZipArchiveReader::IsValid() const
+{
+	return false;
+}
+
+TArray<FString> FZipArchiveReader::GetFileNames() const
+{
+	return TArray<FString>();
+}
+
+bool FZipArchiveReader::TryReadFile(FStringView FileName, TArray<uint8>& OutData, FOutputDevice* ErrorHandler) const
+{
+	if (ErrorHandler)
+	{
+		ErrorHandler->Log(LogZipArchive.GetCategoryName(), ELogVerbosity::Display, TEXT("Not yet implemented on PLATFORM_MAC"));
+	}
+	return false;
+}
+
+#else // PLATFORM_MAC
+
 #include "Containers/StringConv.h"
 #include "Containers/Map.h"
 #include "GenericPlatform/GenericPlatformFile.h"
@@ -335,4 +375,6 @@ zip_int64_t FZipArchiveReader::FImpl::ZipSourceFunctionReader(
 	}
 }
 
+#endif // PLATFORM_MAC
+
 #endif
\ No newline at end of file
