ZipArchiveReader: Fix derefence of uninitialized data; ZipError structures need to be initialized before being passed in to zip functions.
#rnx

diff --git a/Engine/Source/Developer/FileUtilities/Private/ZipArchiveReader.cpp b/Engine/Source/Developer/FileUtilities/Private/ZipArchiveReader.cpp
--- a/Engine/Source/Developer/FileUtilities/Private/ZipArchiveReader.cpp
+++ b/Engine/Source/Developer/FileUtilities/Private/ZipArchiveReader.cpp
@@ -78,6 +78,7 @@ FZipArchiveReader::FImpl::FImpl(IFileHandle* InFileHandle, FOutputDevice* ErrorH
 	FilePos = 0;
 	FileSize = FileHandle->Size();
 	zip_error_t ZipError;
+	zip_error_init(&ZipError);
 	ZipFileSource = zip_source_function_create(ZipSourceFunctionReaderStatic, this, &ZipError);
 	if (!ZipFileSource)
 	{
@@ -91,6 +92,7 @@ FZipArchiveReader::FImpl::FImpl(IFileHandle* InFileHandle, FOutputDevice* ErrorH
 		return;
 	}
 
+	zip_error_init(&ZipError);
 	ZipFile = zip_open_from_source(ZipFileSource, ZIP_RDONLY, &ZipError);
 	if (!ZipFile)
 	{
