ZipArchiveWriter: Fix crash in FZipArchiveWriter::Write when Size argument is 0. Emit a failed ensure when Filename is empty in AddFile.
#rnx

diff --git a/Engine/Source/Developer/FileUtilities/Private/ZipArchiveWriter.cpp b/Engine/Source/Developer/FileUtilities/Private/ZipArchiveWriter.cpp
--- a/Engine/Source/Developer/FileUtilities/Private/ZipArchiveWriter.cpp
+++ b/Engine/Source/Developer/FileUtilities/Private/ZipArchiveWriter.cpp
@@ -101,6 +101,10 @@ FZipArchiveWriter::~FZipArchiveWriter()
 
 void FZipArchiveWriter::AddFile(const FString& Filename, TConstArrayView<uint8> Data, const FDateTime& Timestamp)
 {
+	if (!ensureMsgf(!Filename.IsEmpty(), TEXT("Failed to write data to zip file; filename is empty.")))
+	{
+		return;
+	}
 	uint32 Crc = FCrc::MemCrc32(Data.GetData(), Data.Num());
 
 	// Convert the date-time to a zip file timestamp (2-second resolution).
diff --git a/Engine/Source/Developer/FileUtilities/Public/FileUtilities/ZipArchiveWriter.h b/Engine/Source/Developer/FileUtilities/Public/FileUtilities/ZipArchiveWriter.h
--- a/Engine/Source/Developer/FileUtilities/Public/FileUtilities/ZipArchiveWriter.h
+++ b/Engine/Source/Developer/FileUtilities/Public/FileUtilities/ZipArchiveWriter.h
@@ -45,8 +45,11 @@ class FILEUTILITIES_API FZipArchiveWriter
 	inline void Write(uint64 V) { Write((void*)&V, sizeof(V)); }
 	inline void Write(void* Src, uint64 Size)
 	{
-		void* Dst = &Buffer[Buffer.AddUninitialized(Size)];
-		FMemory::Memcpy(Dst, Src, Size);
+		if (Size)
+		{
+			void* Dst = &Buffer[Buffer.AddUninitialized(Size)];
+			FMemory::Memcpy(Dst, Src, Size);
+		}
 	}
 	inline uint64 Tell() { return (File ? File->Tell() : 0) + Buffer.Num(); }
 	void Flush();
