Updated ZipArchiveWriter to allow UTF-8 for file and directory names within the archive

diff --git a/Engine/Source/Developer/FileUtilities/Private/ZipArchiveWriter.cpp b/Engine/Source/Developer/FileUtilities/Private/ZipArchiveWriter.cpp
--- a/Engine/Source/Developer/FileUtilities/Private/ZipArchiveWriter.cpp
+++ b/Engine/Source/Developer/FileUtilities/Private/ZipArchiveWriter.cpp
@@ -4,6 +4,7 @@
 
 #if WITH_ENGINE
 
+#include "Containers/Utf8String.h"
 #include "GenericPlatform/GenericPlatformFile.h"
 #include "ZipArchivePrivate.h"
 
@@ -25,25 +26,37 @@ FZipArchiveWriter::~FZipArchiveWriter()
 	uint64 DirStartOffset = Tell();
 	for (FFileEntry& Entry : Files)
 	{
+		// Central directory File header: (from specification linked above)
 		const static uint8 Footer[] =
 		{
-			0x50, 0x4b, 0x01, 0x02, 0x3f, 0x00, 0x2d, 0x00,
-			0x00, 0x00, 0x00, 0x00
+			0x50, 0x4b, 0x01, 0x02,  // Central file header signature
+			0x3f, 0x00,  // Version made by (MS-DOS - v6.3)
+			0x2d, 0x00,  // Version needed to extract (MS-DOS - v4.5)
+			0x00, 0x08,  // General purpose bit flag (Language encoding flag = 1)
+			0x00, 0x00  // Compression method (none)
 		};
 		Write((void*)Footer, sizeof(Footer));
 		Write(Entry.Time);
 		Write(Entry.Crc32);
 
+		// Compressed and Uncompressed size - unused.
 		Write((uint64)0xffffffffffffffff);
-		Write((uint16)Entry.Filename.Len());
+
+		FUtf8String UTF8Filename = *Entry.Filename;
+		Write((uint16)UTF8Filename.Len());
 		const static uint8 Fields[] =
 		{
-			0x20, 0x00, 0x00, 0x00, 0x00, 0x00, 0x00, 0x00,
-			0x20, 0x00, 0x00, 0x00, 0xff, 0xff, 0xff, 0xff
+			0x20, 0x00, // Length of extra fields (Zip64 Extended Information)
+			0x00, 0x00, // File comment length
+			0x00, 0x00, // Disk number start
+			0x00, 0x00, // Internal file attributes
+			0x20, 0x00, 0x00, 0x00, // External file attributes
+			0xff, 0xff, 0xff, 0xff // Relative offset of local header
 		};
 		Write((void*)Fields, sizeof(Fields));
-		Write((void*)TCHAR_TO_UTF8(*Entry.Filename), Entry.Filename.Len());
+		Write((void*)GetData(UTF8Filename), UTF8Filename.Len());
 
+		// Zip64 Extended Information block
 		Write((uint16)0x01);
 		Write((uint16)0x1c);
 
@@ -120,20 +133,27 @@ void FZipArchiveWriter::AddFile(const FString& Filename, TConstArrayView<uint8>
 
 	FFileEntry* Entry = new (Files) FFileEntry(Filename, Crc, Data.Num(), FileOffset, ZipTime);
 
+	// Local File Header
 	static const uint8 Header[] =
 	{
-		0x50, 0x4b, 0x03, 0x04, 0x2d, 0x00, 0x00, 0x00,
-		0x00, 0x00
+		0x50, 0x4b, 0x03, 0x04, // Local file header signature
+		0x2d, 0x00, // Version needed to extract (MS DOS - v4.5)
+		0x00, 0x08, // General purpose bit flag (Language encoding flag = 1)
+		0x00, 0x00 // Compression method (none)
 	};
 	Write((void*)Header, sizeof(Header));
 	Write(ZipTime);
 	Write(Crc);
+
+	// Compressed and Uncompressed size - unused.
 	Write((uint64)0xffffffffffffffff);
-	Write((uint16)Filename.Len());
-	Write((uint16)0x20);
 
-	Write((void*)TCHAR_TO_UTF8(*Entry->Filename), Filename.Len());
+	FUtf8String UTF8Filename = *Entry->Filename;
+	Write((uint16)UTF8Filename.Len());
+	Write((uint16)0x20); // Length of extra fields (Zip64 Extended Information)
+	Write((void*)GetData(UTF8Filename), UTF8Filename.Len());
 
+	// Zip64 Extended Information block
 	Write((uint16)0x01);
 	Write((uint16)0x1c);
 	Write((uint64)Data.Num());
