Fix ZipArchiveWritter ZIP64 header output

Also adds more comments explaining the different fields of the zip file and a test to make sure data can be extracted from the zip file

diff --git a/Engine/Source/Developer/FileUtilities/Private/Tests/ZipArchiveTest.cpp b/Engine/Source/Developer/FileUtilities/Private/Tests/ZipArchiveTest.cpp
new file mode 100644
--- /dev/null
+++ b/Engine/Source/Developer/FileUtilities/Private/Tests/ZipArchiveTest.cpp
@@ -0,0 +1,83 @@
+// Copyright Epic Games, Inc. All Rights Reserved.
+
+#include "Misc/AutomationTest.h"
+#include "Misc/Paths.h"
+#include "Misc/FileHelper.h"
+#include "HAL/PlatformFileManager.h"
+
+#include "FileUtilities/ZipArchiveWriter.h"
+#include "FileUtilities/ZipArchiveReader.h"
+
+#if WITH_DEV_AUTOMATION_TESTS
+
+#if WITH_ENGINE
+
+IMPLEMENT_SIMPLE_AUTOMATION_TEST(FZipArchiveTest, "FileUtilities.ZipArchive", EAutomationTestFlags_ApplicationContextMask | EAutomationTestFlags::EngineFilter);
+
+bool FZipArchiveTest::RunTest(const FString& InParameter)
+{	
+	const FString TempDir = FPaths::AutomationTransientDir();
+	const FString Prefix = TEXT("ZipArchiveTest");
+	const FString TxtExtension = TEXT(".txt");
+	const FString TempFileToZip = FPaths::CreateTempFilename(*TempDir, *Prefix, *TxtExtension);
+
+	// Contents to be zipped
+	const FString FileContents = TEXT("FileUtilities ZipArchive Test");
+
+	const FString ZipExtension = TEXT(".zip");
+	const FString ZipFilePath = FPaths::ConvertRelativePathToFull(FPaths::CreateTempFilename(*TempDir, *Prefix, *ZipExtension));
+	const FString TestDirectory = FPaths::GetPath(ZipFilePath);
+
+	// Make sure the directory where OpenWrite is called exists
+	const bool bMakeTree = true;
+	UTEST_TRUE("Making directory tree", IFileManager::Get().MakeDirectory(*TestDirectory, bMakeTree));
+
+	ON_SCOPE_EXIT
+	{
+		// Make sure the Tmp folder gets deleted when the tests finishes
+		const bool bRequireExists = true;
+		const bool bRemoveTree = true;
+		IFileManager::Get().DeleteDirectory(*TestDirectory, bRequireExists, bRemoveTree);
+	};
+
+	IPlatformFile& PlatformFile = FPlatformFileManager::Get().GetPlatformFile();
+	
+	{
+		// Creates a zip file
+
+		IFileHandle* ZipFile = PlatformFile.OpenWrite(*ZipFilePath);
+		UTEST_NOT_NULL("Zip File is valid", ZipFile);
+
+		FZipArchiveWriter ZipWriter{ ZipFile };
+
+		auto ANSIFileContents = StringCast<ANSICHAR>(*FileContents);
+		TConstArrayView<uint8> StringView((uint8*) ANSIFileContents.Get(), ANSIFileContents.Length());
+		ZipWriter.AddFile(FPaths::GetCleanFilename(TempFileToZip), StringView, FDateTime::Now());
+	}
+	
+	{
+		// Reads the zip file and see if the contents are correct
+
+		IFileHandle* ZipFile = PlatformFile.OpenRead(*ZipFilePath);
+		UTEST_NOT_NULL("Zip File is valid", ZipFile);
+
+		FZipArchiveReader ZipReader{ ZipFile };
+		const TArray<FString> FileNames = ZipReader.GetFileNames();
+		UTEST_EQUAL("File Count", FileNames.Num(), 1);
+
+		for (const FString& FileName : FileNames)
+		{
+			TArray<uint8> FileContentsBuffer;
+			UTEST_TRUE("Try Read File From Zip", ZipReader.TryReadFile(FileName, FileContentsBuffer));
+
+			TConstArrayView<ANSICHAR> StringView((ANSICHAR*) FileContentsBuffer.GetData(), FileContentsBuffer.Num());
+			UTEST_EQUAL("Are Contents the Same", FString{ StringView }, FileContents);
+		}
+	}
+
+	return true;
+}
+
+#endif // WITH_ENGINE
+
+#endif // WITH_DEV_AUTOMATION_TESTS
\ No newline at end of file
diff --git a/Engine/Source/Developer/FileUtilities/Private/ZipArchiveWriter.cpp b/Engine/Source/Developer/FileUtilities/Private/ZipArchiveWriter.cpp
--- a/Engine/Source/Developer/FileUtilities/Private/ZipArchiveWriter.cpp
+++ b/Engine/Source/Developer/FileUtilities/Private/ZipArchiveWriter.cpp
@@ -46,24 +46,23 @@ FZipArchiveWriter::~FZipArchiveWriter()
 		Write((uint16)UTF8Filename.Len());
 		const static uint8 Fields[] =
 		{
-			0x20, 0x00, // Length of extra fields (Zip64 Extended Information)
+			0x1c, 0x00, // Length of extra fields (Zip64 Extended Information)
 			0x00, 0x00, // File comment length
 			0x00, 0x00, // Disk number start
 			0x00, 0x00, // Internal file attributes
 			0x20, 0x00, 0x00, 0x00, // External file attributes
-			0xff, 0xff, 0xff, 0xff // Relative offset of local header
+			0xff, 0xff, 0xff, 0xff // Relative offset of local header (set to 0xff as it is provided in the Zip64 block)
 		};
 		Write((void*)Fields, sizeof(Fields));
 		Write((void*)GetData(UTF8Filename), UTF8Filename.Len());
 
 		// Zip64 Extended Information block
-		Write((uint16)0x01);
-		Write((uint16)0x1c);
+		Write((uint16)0x01); // Tag the Zip64
+		Write((uint16)0x18); // Size of this block (24 bytes)
 
-		Write((uint64)Entry.Length);
-		Write((uint64)Entry.Length);
-		Write((uint64)Entry.Offset);
-		Write((uint32)0);
+		Write((uint64)Entry.Length); // Uncompressed size
+		Write((uint64)Entry.Length); // Compressed Size
+		Write((uint64)Entry.Offset); // Offset of local header record
 
 		Flush();
 	}
@@ -74,31 +73,40 @@ FZipArchiveWriter::~FZipArchiveWriter()
 	// Write ZIP64 end of central directory record
 	const static uint8 Record[] =
 	{
-		0x50, 0x4b, 0x06, 0x06, 0x2c, 0x00, 0x00, 0x00,
-		0x00, 0x00, 0x00, 0x00, 0x2d, 0x00, 0x2d, 0x00,
-		0x00, 0x00, 0x00, 0x00, 0x00, 0x00, 0x00, 0x00,
+		0x50, 0x4b, 0x06, 0x06,							// Zip64 end of central directory record signature
+		0x2c, 0x00, 0x00, 0x00, 0x00, 0x00, 0x00, 0x00, // Size of the end of central directory record
+		0x2d, 0x00,										// Version Creator (MS-DOS - v4.5)
+		0x2d, 0x00,										// Version Viewer (MS-DOS - v4.5)
+		0x00, 0x00, 0x00, 0x00,							// Disk Number
+		0x00, 0x00, 0x00, 0x00,							// Disk with central directory
 	};
 	Write((void*)Record, sizeof(Record));
-	Write((uint64)Files.Num());
-	Write((uint64)Files.Num());
-	Write(DirectorySizeInBytes);
-	Write(DirStartOffset);
+	Write((uint64)Files.Num());							// Number of central directory records
+	Write((uint64)Files.Num());							// Total number of records
+	Write(DirectorySizeInBytes);						// Size of central directory
+	Write(DirStartOffset);								// Offset of central directory
 
 	// Write ZIP64 end of central directory locator
 	const static uint8 Locator[] =
 	{
-		0x50, 0x4b, 0x06, 0x07, 0x00, 0x00, 0x00, 0x00,
+		0x50, 0x4b, 0x06, 0x07,		// Zip64 end of central directory locator signature
+		0x00, 0x00, 0x00, 0x00,		// Disk with end of central directory record
 	};
 	Write((void*)Locator, sizeof(Locator));
-	Write(DirEndOffset);
-	Write((uint32)0x01);
+	Write(DirEndOffset);			// Offset of end of central directory
+	Write((uint32)0x01);			// Total number of disks
 
 	// Write normal end of central directory record
 	const static uint8 EndRecord[] =
 	{
-		0x50, 0x4b, 0x05, 0x06, 0xff, 0xff, 0xff, 0xff,
-		0xff, 0xff, 0xff, 0xff, 0xff, 0xff, 0xff, 0xff,
-		0xff, 0xff, 0xff, 0xff, 0x00, 0x00
+		0x50, 0x4b, 0x05, 0x06,		// End of central directory record signature
+		0x00, 0x00,					// Number of this disk
+		0x00, 0x00,					// Number of the disk with the start of the central directory
+		0xff, 0xff,					// Total number of entries in the central directory on this disk
+		0xff, 0xff,					// Total number of entries
+		0xff, 0xff, 0xff, 0xff,		// Size of central directory
+		0xff, 0xff, 0xff, 0xff,		// Offset of central directory
+		0x00, 0x00					// comment length
 	};
 	Write((void*)EndRecord, sizeof(EndRecord));
 
@@ -145,21 +153,19 @@ void FZipArchiveWriter::AddFile(const FString& Filename, TConstArrayView<uint8>
 	Write(ZipTime);
 	Write(Crc);
 
-	// Compressed and Uncompressed size - unused.
+	// Compressed and Uncompressed size - (set to 0xff as it is provided by the Zip64 block).
 	Write((uint64)0xffffffffffffffff);
 
 	FUtf8String UTF8Filename = *Entry->Filename;
 	Write((uint16)UTF8Filename.Len());
-	Write((uint16)0x20); // Length of extra fields (Zip64 Extended Information)
+	Write((uint16)0x14); // Length of extra fields (Zip64 Extended Information)
 	Write((void*)GetData(UTF8Filename), UTF8Filename.Len());
 
 	// Zip64 Extended Information block
-	Write((uint16)0x01);
-	Write((uint16)0x1c);
-	Write((uint64)Data.Num());
-	Write((uint64)Data.Num());
-	Write((uint64)FileOffset);
-	Write((uint32)0);
+	Write((uint16)0x01);		// Zip64 tag
+	Write((uint16)0x10);		// Size of this block (16 bytes)
+	Write((uint64)Data.Num());  // Uncompressed size
+	Write((uint64)Data.Num());  // Compressed size
 
 	Write((void*)Data.GetData(), Data.Num());
 
