ZipArchiveReader: Fix bug in our call to libzip API function zip_source_make_command_bitmap. The list of arguments needs to be terminated with -1, per the libzip documentation. A licensee reported this fix was necessary to fix incorrect behavior on some platforms.
#rnx

diff --git a/Engine/Source/Developer/FileUtilities/Private/ZipArchiveReader.cpp b/Engine/Source/Developer/FileUtilities/Private/ZipArchiveReader.cpp
--- a/Engine/Source/Developer/FileUtilities/Private/ZipArchiveReader.cpp
+++ b/Engine/Source/Developer/FileUtilities/Private/ZipArchiveReader.cpp
@@ -331,7 +331,7 @@ zip_int64_t FZipArchiveReader::FImpl::ZipSourceFunctionReader(
 		return static_cast<zip_int64_t>(FilePos);
 	case ZIP_SOURCE_SUPPORTS:
 		return zip_source_make_command_bitmap(ZIP_SOURCE_OPEN, ZIP_SOURCE_READ, ZIP_SOURCE_CLOSE, ZIP_SOURCE_STAT,
-			ZIP_SOURCE_ERROR, ZIP_SOURCE_FREE, ZIP_SOURCE_SEEK, ZIP_SOURCE_TELL, ZIP_SOURCE_SUPPORTS);
+			ZIP_SOURCE_ERROR, ZIP_SOURCE_FREE, ZIP_SOURCE_SEEK, ZIP_SOURCE_TELL, ZIP_SOURCE_SUPPORTS, -1);
 	default:
 		return 0;
 	}
