ZipArchiveWriter: Truncation warning fix.

Skein includes this header but it causes errors, due to Skein using UnsafeTypeCastWarningLevel = WarningLevel.Error.

Note that this doesn't change the unsafe behavior!
Although it's now checkf'd in developer builds. The proper way would be to use a TArray64, but since this is public it would require going through the deprecation flow, which is out of scope for what we're trying to do (e.g.: get rid of the warning).

#rnx

diff --git a/Engine/Source/Developer/FileUtilities/Public/FileUtilities/ZipArchiveWriter.h b/Engine/Source/Developer/FileUtilities/Public/FileUtilities/ZipArchiveWriter.h
--- a/Engine/Source/Developer/FileUtilities/Public/FileUtilities/ZipArchiveWriter.h
+++ b/Engine/Source/Developer/FileUtilities/Public/FileUtilities/ZipArchiveWriter.h
@@ -47,7 +47,7 @@ class FILEUTILITIES_API FZipArchiveWriter
 	{
 		if (Size)
 		{
-			void* Dst = &Buffer[Buffer.AddUninitialized(Size)];
+			void* Dst = &Buffer[Buffer.AddUninitialized(IntCastChecked<int32>(Size))];
 			FMemory::Memcpy(Dst, Src, Size);
 		}
 	}
