Wrap FZipArchiveReader test with WITH_EDITOR since FZipArchiveReader is editor only

diff --git a/Engine/Source/Developer/FileUtilities/Private/Tests/ZipArchiveTest.cpp b/Engine/Source/Developer/FileUtilities/Private/Tests/ZipArchiveTest.cpp
--- a/Engine/Source/Developer/FileUtilities/Private/Tests/ZipArchiveTest.cpp
+++ b/Engine/Source/Developer/FileUtilities/Private/Tests/ZipArchiveTest.cpp
@@ -53,8 +53,13 @@ bool FZipArchiveTest::RunTest(const FString& InParameter)
 		auto ANSIFileContents = StringCast<ANSICHAR>(*FileContents);
 		TConstArrayView<uint8> StringView((uint8*) ANSIFileContents.Get(), ANSIFileContents.Length());
 		ZipWriter.AddFile(FPaths::GetCleanFilename(TempFileToZip), StringView, FDateTime::Now());
+
 	}
-	
+
+#if WITH_EDITOR
+
+	// FZipArchiveReader is editor only
+
 	{
 		// Reads the zip file and see if the contents are correct
 
@@ -75,6 +80,8 @@ bool FZipArchiveTest::RunTest(const FString& InParameter)
 		}
 	}
 
+#endif // WITH_EDITOR
+
 	return true;
 }
 
