Added the possibility to compress the shader symbols, usefull on patforms where the debug info is not compressed. ZipWriterArchive can support compress or uncompress data. Added Automation test to validate the new compression option on zip archives.

#tests Build shadersSymbols + ZipArchiveTest

diff --git a/Engine/Source/Developer/FileUtilities/Private/Tests/ZipArchiveTest.cpp b/Engine/Source/Developer/FileUtilities/Private/Tests/ZipArchiveTest.cpp
--- a/Engine/Source/Developer/FileUtilities/Private/Tests/ZipArchiveTest.cpp
+++ b/Engine/Source/Developer/FileUtilities/Private/Tests/ZipArchiveTest.cpp
@@ -1,12 +1,16 @@
 // Copyright Epic Games, Inc. All Rights Reserved.
 
 #include "Misc/AutomationTest.h"
-#include "Misc/Paths.h"
+#include "Misc/CommandLine.h"
 #include "Misc/FileHelper.h"
-#include "HAL/PlatformFileManager.h"
-
+#include "Misc/Paths.h"
 #include "FileUtilities/ZipArchiveWriter.h"
 #include "FileUtilities/ZipArchiveReader.h"
+#include "HAL/PlatformFileManager.h"
+#include "Logging/LogMacros.h"
+
+
+DEFINE_LOG_CATEGORY_STATIC(LogAutomationZipArchive, Log, All);
 
 #if WITH_DEV_AUTOMATION_TESTS
 
@@ -81,6 +85,118 @@ bool FZipArchiveTest::RunTest(const FString& InParameter)
 	}
 
 #endif // WITH_EDITOR
+	return true;
+}
+
+
+IMPLEMENT_SIMPLE_AUTOMATION_TEST(FZipArchiveCompressTest, "FileUtilities.ZipArchiveCompress", EAutomationTestFlags_ApplicationContextMask | EAutomationTestFlags::EngineFilter);
+
+bool FZipArchiveCompressTest::RunTest(const FString& InParameter)
+{
+	IFileManager& FileManager = IFileManager::Get();
+	IPlatformFile& PlatformFile = FPlatformFileManager::Get().GetPlatformFile();
+
+	FString FileExt = TEXT("");
+	FString FileDir = TEXT("");
+	bool bCompress = false;
+	FString CmdLine = FCommandLine::Get();
+	FParse::Value(*CmdLine, TEXT("zipArchiveExt="), FileExt, false);
+	FParse::Value(*CmdLine, TEXT("zipArchiveDir="), FileDir, false);
+	bCompress = FParse::Param(FCommandLine::Get(), TEXT("zipArchiveCompress"));
+
+	// Didn't find any command line arguments, exit the test.
+	if (FileExt.IsEmpty() || FileDir.IsEmpty())
+	{
+		return true;
+	}
+
+	const FString TestDirectory = FPaths::AutomationTransientDir();
+	PlatformFile.CreateDirectory(*TestDirectory);
+
+	const FString ZipFilePath = FPaths::ConvertRelativePathToFull(FPaths::CreateTempFilename(*TestDirectory, TEXT("ZipArchiveCompressTest"), TEXT(".zip")));
+
+	// Make sure the directory where OpenWrite is called exists
+	const bool bMakeTree = true;
+	UTEST_TRUE("Making directory tree", IFileManager::Get().MakeDirectory(*TestDirectory, bMakeTree));
+
+	ON_SCOPE_EXIT
+	{
+		// Make sure the Tmp folder gets deleted when the tests finishes
+		const bool bRequireExists = true;
+		const bool bRemoveTree = true;
+		FileManager.DeleteDirectory(*TestDirectory, bRequireExists, bRemoveTree);
+	};
+
+	// create the compress zip file
+	{
+		// Find all the files in the directory
+		TArray<FString> FileToCompress;
+		PlatformFile.FindFiles(FileToCompress, *FileDir, *FileExt);
+
+		// Creates a zip file
+		IFileHandle* ZipFile = PlatformFile.OpenWrite(*ZipFilePath);
+		UTEST_NOT_NULL("Zip File is valid", ZipFile);
+
+		EZipArchiveOptions ZipOptions = EZipArchiveOptions::RemoveDuplicate;
+		if (bCompress)
+		{
+			ZipOptions |= EZipArchiveOptions::Deflate;
+		}
+
+		FZipArchiveWriter ZipWriter(ZipFile, ZipOptions);
+
+		for (int Index = 0; Index < FileToCompress.Num(); Index++)
+		{
+			TUniquePtr<FArchive> Reader = TUniquePtr<FArchive>(FileManager.CreateFileReader(*FileToCompress[Index]));
+			UTEST_TRUE("Fail to create the File archive reader", Reader.IsValid());
+
+			int64 Size = Reader->TotalSize();
+			TArray<uint8> RawData;
+			RawData.AddUninitialized(Size);
+			Reader->Serialize(RawData.GetData(), Size);
+			Reader->Close();
+			ZipWriter.AddFile(FileToCompress[Index], RawData, FDateTime::Now());
+		}
+	}
+
+	// Reads the zip file and see if the contents are correct
+	{
+		IFileHandle* ZipFile = PlatformFile.OpenRead(*ZipFilePath);
+		UTEST_NOT_NULL("Zip File is valid", ZipFile);
+
+		FZipArchiveReader ZipReader{ ZipFile };
+		const TArray<FString> FileNames = ZipReader.GetFileNames();
+		int32 NbProcessFiles = 0;
+		for (const FString& FileName : FileNames)
+		{
+			// open the source uncompress file 
+			UTEST_TRUE("File name doesn't match", FileManager.FileExists(*FileName));
+
+			TUniquePtr<FArchive> UnCompressReader = TUniquePtr<FArchive>(FileManager.CreateFileReader(*FileName));
+			UTEST_TRUE("Fail read the uncompress file", UnCompressReader.IsValid());
+
+			// extract the file in the zip archive
+			TArray<uint8> CompressData;
+			UTEST_TRUE("Try Read File From Zip", ZipReader.TryReadFile(FileName, CompressData));
+
+			TConstArrayView<ANSICHAR> CompressDataView((ANSICHAR*)CompressData.GetData(), CompressData.Num());
+			int64 Size = UnCompressReader->TotalSize();
+			TArray<uint8> UnCompressData;
+			UnCompressData.AddUninitialized(Size);
+			UnCompressReader->Serialize(UnCompressData.GetData(), Size);
+			UnCompressReader->Close();
+			TConstArrayView<ANSICHAR> UnCompressView((ANSICHAR*)UnCompressData.GetData(), UnCompressData.Num());
+
+			FString StringCompressView(CompressDataView);
+			FString StringUnCompressView(CompressDataView);
+			UTEST_EQUAL("Are Contents the Same", FString{ CompressDataView }, FString{ UnCompressView });
+			NbProcessFiles++;
+		}
+		if (NbProcessFiles == FileNames.Num())
+		{
+			UE_LOG(LogAutomationZipArchive, Display, TEXT("Test succeded"));
+		}
+	}
 
 	return true;
 }
diff --git a/Engine/Source/Developer/FileUtilities/Private/ZipArchiveWriter.cpp b/Engine/Source/Developer/FileUtilities/Private/ZipArchiveWriter.cpp
--- a/Engine/Source/Developer/FileUtilities/Private/ZipArchiveWriter.cpp
+++ b/Engine/Source/Developer/FileUtilities/Private/ZipArchiveWriter.cpp
@@ -6,109 +6,200 @@
 
 #include "Containers/Utf8String.h"
 #include "GenericPlatform/GenericPlatformFile.h"
+#include "libzip/zip.h"
+#include "Misc/Compression.h"
 #include "ZipArchivePrivate.h"
 
+
 DEFINE_LOG_CATEGORY(LogZipArchive);
 
-FZipArchiveWriter::FZipArchiveWriter(IFileHandle* InFile)
-	: File(InFile)
+// Zip File Format Specification (We dont use encryption and data descriptor)
+// https://www.loc.gov/preservation/digital/formats/digformatspecs/APPNOTE%2820120901%29_Version_6.3.3.txt
+//
+// Zip64 Extended Information block
+#pragma pack(push,1)
+struct FZip64ExtraFieldHeader
+{
+	uint16 Id = 0x0001; // Tag the Zip64
+	uint16 Length = 16; // Size of the data (16 bytes = 2 x uint64)
+	uint64 UnCompressedSize;
+	uint64 CompressedSize;
+};
+#pragma pack(pop)
+
+#pragma pack(push,1)
+struct FZipCDHeaderExtraFields : public FZip64ExtraFieldHeader
+{
+	uint64 Offset; // ByteOffset from beginning of zipfile of the bytes of the FileEntry
+};
+#pragma pack(pop)
+
+enum class EZipGenPurposeFlags : uint16
+{
+	LanguageEncodingFlag = 1 << 11
+};
+
+// Local File Header
+#pragma pack(push,1)
+struct FZipLocalHeader
+{
+	uint8 Sig[4] = { 0x50, 0x4b, 0x03, 0x04 }; // Local file header signature
+	uint16 Version = 45; // Version needed to extract (MS DOS - v4.5)
+	uint16 GenPurposeBit; // General purpose bit flag (Language encoding flag = 1)
+	uint16 CompMode;
+	uint32 TimeDate;
+	uint32 Crc;
+	uint32 CompressedSize; // 0xFFFFFFFF if Zip64 format
+	uint32 UnCompressedSize; // 0xFFFFFFFF if Zip64 format
+	uint16 FileNameLength;
+	uint16 ExtraFieldLength;// if zip64Format = sizeof(FZip64ExtraFieldHeader); 
+
+};
+#pragma pack(pop)
+
+enum class EZipExternalFileAttributeFlags : uint32
+{
+	MSDOS_Directory_Archive = 1 << 5
+};
+
+// Central directory header
+#pragma pack(push,1)
+struct FZipCDHeader
+{
+	uint8 Sig[4] = { 0x50, 0x4b, 0x01, 0x02 }; // Central file header signature
+	uint16 VersionMade = 63; // Version made by (MS-DOS - v6.3)
+	uint16 VersionNeeded = 45; // Version needed to extract (MS-DOS - v4.5)
+	uint16 GenPurposeBit; // General purpose bit flag
+	uint16 CompMode;
+	uint32 TimeDate;
+	uint32 Crc;
+	uint32 CompressedSize; // 0xFFFFFFFF if Zip64 format
+	uint32 UnCompressedSize; // 0xFFFFFFFF if Zip64 format
+	uint16 FilenameLength;
+	uint16 ExtraFieldLength;// = sizeof(FZip64ExtraFieldHeader);// Length of extra fields (Zip64 Extended Information)
+	uint16 FileCommentLength = 0;
+	uint16 DiskNumberStart = 0;
+	uint16 InternalFileAttr = 0;
+	uint32 ExternalFileAttr = (uint32)EZipExternalFileAttributeFlags::MSDOS_Directory_Archive;
+	uint32 RelativeLocHeaderOffset;// 0xFFFFFFFF if Zip64 format 
+};
+#pragma pack(pop)
+
+// ZIP64 end of central directory record header
+#pragma pack(push,1)
+struct FZip64EndOfCDRecord
+{
+	uint8 Sig[4] = { 0x50, 0x4b, 0x06, 0x06 }; // Zip64 end of central directory record signature
+	uint64 SizeOfEndOfCDR = 0x2c; // Size of the end of central directory record
+	uint16 VersionMade = 63; // Version made by (MS-DOS - v6.3)
+	uint16 VersionNeeded = 45; // Version Viewer (MS-DOS - v4.5)
+	uint32 DiskNumber = 0;
+	uint32 CDDiskNumber = 0;
+	uint64 CDRecords;
+	uint64 CDTotalRecords;
+	uint64 CDSize;
+	uint64 CDStartOffset;
+};
+#pragma pack(pop)
+
+//  ZIP64 end of central directory locator header
+#pragma pack(push,1)
+struct FZip64EndOfCDLocator
+{
+	uint8 Sig[4] = { 0x50, 0x4b, 0x06, 0x07 }; // Zip64 end of central directory locator signature
+	uint32 DiskNumber = 0; // Disk with end of central directory record
+	uint64 EndOffsetCD; // Offset of end of central directory
+	uint32 TotalDiskNumber = 1;
+};
+#pragma pack(pop)
+
+// Legacy(non - ZIP64) header for the End of Central Directory
+#pragma pack(push,1)
+struct FZipEndOfCDRecord
+{
+	uint8 Sig[4] = { 0x50, 0x4b, 0x05, 0x06 }; // End of central directory record signature
+	uint16 DiskNumber = 0xFFFF; //0xFFFF means read the value from the FZip64EndOfCDRecord
+	uint16 CDDiskNumber = 0xFFFF; //0xFFFF means read the value from the FZip64EndOfCDRecord
+	uint16 CDRecords = 0xFFFF; //0xFFFF means read the value from the FZip64EndOfCDRecord
+	uint16 TotalCDRecords = 0xFFFF; //0xFFFF means read the value from the FZip64EndOfCDRecord
+	uint32 CDSize = 0xFFFFFFFF; //0xFFFF means read the value from the FZip64EndOfCDRecord
+	uint32 CDOffset = 0xFFFFFFFF; //0xFFFF means read the value from the FZip64EndOfCDRecord
+	uint16 CommentLength = 0;
+};
+#pragma pack(pop)
+
+FZipArchiveWriter::FZipArchiveWriter(IFileHandle* InFile, EZipArchiveOptions InZipOptions)
+	: File(InFile),
+	ZipOptions(InZipOptions)
 {
 }
 
 FZipArchiveWriter::~FZipArchiveWriter()
 {
-	// Zip File Format Specification:
-	// https://www.loc.gov/preservation/digital/formats/digformatspecs/APPNOTE%2820120901%29_Version_6.3.3.txt
-
 	UE_LOG(LogZipArchive, Display, TEXT("Closing zip file with %d entries."), Files.Num());
 
-	// Write the file directory
+	// Write the file central directory
 	uint64 DirStartOffset = Tell();
+
 	for (FFileEntry& Entry : Files)
 	{
-		// Central directory File header: (from specification linked above)
-		const static uint8 Footer[] =
-		{
-			0x50, 0x4b, 0x01, 0x02,  // Central file header signature
-			0x3f, 0x00,  // Version made by (MS-DOS - v6.3)
-			0x2d, 0x00,  // Version needed to extract (MS-DOS - v4.5)
-			0x00, 0x08,  // General purpose bit flag (Language encoding flag = 1)
-			0x00, 0x00  // Compression method (none)
-		};
-		Write((void*)Footer, sizeof(Footer));
-		Write(Entry.Time);
-		Write(Entry.Crc32);
-
-		// Compressed and Uncompressed size - unused.
-		Write((uint64)0xffffffffffffffff);
-
 		FUtf8String UTF8Filename = *Entry.Filename;
-		Write((uint16)UTF8Filename.Len());
-		const static uint8 Fields[] =
-		{
-			0x1c, 0x00, // Length of extra fields (Zip64 Extended Information)
-			0x00, 0x00, // File comment length
-			0x00, 0x00, // Disk number start
-			0x00, 0x00, // Internal file attributes
-			0x20, 0x00, 0x00, 0x00, // External file attributes
-			0xff, 0xff, 0xff, 0xff // Relative offset of local header (set to 0xff as it is provided in the Zip64 block)
-		};
-		Write((void*)Fields, sizeof(Fields));
+		///////////////////////////
+		// central directory header
+		///////////////////////////
+		FZipCDHeader ZipCDHeader;
+		ZipCDHeader.GenPurposeBit = (uint16)EZipGenPurposeFlags::LanguageEncodingFlag;
+		ZipCDHeader.CompMode = Entry.bIsCompress ? ZIP_CM_DEFLATE : ZIP_CM_STORE;
+		ZipCDHeader.TimeDate = Entry.Time;
+		ZipCDHeader.Crc = Entry.Crc32;
+		ZipCDHeader.UnCompressedSize = 0xFFFFFFFF;
+		ZipCDHeader.CompressedSize = 0xFFFFFFFF;
+		ZipCDHeader.FilenameLength = UTF8Filename.Len();
+		ZipCDHeader.ExtraFieldLength = sizeof(FZipCDHeaderExtraFields);
+		ZipCDHeader.RelativeLocHeaderOffset = 0xFFFFFFFF;
+		Write((void*)&ZipCDHeader, sizeof(FZipCDHeader));
+
+		// Write the variable size data of central directory
 		Write((void*)GetData(UTF8Filename), UTF8Filename.Len());
 
+		///////////////////////////
 		// Zip64 Extended Information block
-		Write((uint16)0x01); // Tag the Zip64
-		Write((uint16)0x18); // Size of this block (24 bytes)
-
-		Write((uint64)Entry.Length); // Uncompressed size
-		Write((uint64)Entry.Length); // Compressed Size
-		Write((uint64)Entry.Offset); // Offset of local header record
-
+		///////////////////////////
+		FZipCDHeaderExtraFields Zip64ExtraField;
+		Zip64ExtraField.Length = sizeof(Zip64ExtraField) - sizeof(Zip64ExtraField.Id) - sizeof(Zip64ExtraField.Length); // Length of the data after the Length field
+		Zip64ExtraField.UnCompressedSize = Entry.UnCompressedSize;
+		Zip64ExtraField.CompressedSize = Entry.CompressedSize;
+		Zip64ExtraField.Offset = Entry.Offset;
+		Write((void*)&Zip64ExtraField, sizeof(FZipCDHeaderExtraFields));
 		Flush();
 	}
-	uint64 DirEndOffset = Tell();
 
-	uint64 DirectorySizeInBytes = DirEndOffset - DirStartOffset;
+	uint64_t EndOffsetCD = Tell();
+	uint64_t DirectorySizeInBytes = EndOffsetCD - DirStartOffset;
 
-	// Write ZIP64 end of central directory record
-	const static uint8 Record[] =
-	{
-		0x50, 0x4b, 0x06, 0x06,							// Zip64 end of central directory record signature
-		0x2c, 0x00, 0x00, 0x00, 0x00, 0x00, 0x00, 0x00, // Size of the end of central directory record
-		0x2d, 0x00,										// Version Creator (MS-DOS - v4.5)
-		0x2d, 0x00,										// Version Viewer (MS-DOS - v4.5)
-		0x00, 0x00, 0x00, 0x00,							// Disk Number
-		0x00, 0x00, 0x00, 0x00,							// Disk with central directory
-	};
-	Write((void*)Record, sizeof(Record));
-	Write((uint64)Files.Num());							// Number of central directory records
-	Write((uint64)Files.Num());							// Total number of records
-	Write(DirectorySizeInBytes);						// Size of central directory
-	Write(DirStartOffset);								// Offset of central directory
 
+	///////////////////////////
+	// Write ZIP64 end of central directory record
+	///////////////////////////
+	FZip64EndOfCDRecord Zip64EndOfCDRecord;
+	Zip64EndOfCDRecord.CDRecords = Files.Num();
+	Zip64EndOfCDRecord.CDTotalRecords = Files.Num();
+	Zip64EndOfCDRecord.CDSize = DirectorySizeInBytes;
+	Zip64EndOfCDRecord.CDStartOffset = DirStartOffset;
+	Write((void*)&Zip64EndOfCDRecord, sizeof(FZip64EndOfCDRecord));
+
+	///////////////////////////
 	// Write ZIP64 end of central directory locator
-	const static uint8 Locator[] =
-	{
-		0x50, 0x4b, 0x06, 0x07,		// Zip64 end of central directory locator signature
-		0x00, 0x00, 0x00, 0x00,		// Disk with end of central directory record
-	};
-	Write((void*)Locator, sizeof(Locator));
-	Write(DirEndOffset);			// Offset of end of central directory
-	Write((uint32)0x01);			// Total number of disks
-
-	// Write normal end of central directory record
-	const static uint8 EndRecord[] =
-	{
-		0x50, 0x4b, 0x05, 0x06,		// End of central directory record signature
-		0x00, 0x00,					// Number of this disk
-		0x00, 0x00,					// Number of the disk with the start of the central directory
-		0xff, 0xff,					// Total number of entries in the central directory on this disk
-		0xff, 0xff,					// Total number of entries
-		0xff, 0xff, 0xff, 0xff,		// Size of central directory
-		0xff, 0xff, 0xff, 0xff,		// Offset of central directory
-		0x00, 0x00					// comment length
-	};
-	Write((void*)EndRecord, sizeof(EndRecord));
+	///////////////////////////
+	FZip64EndOfCDLocator Zip64EndOfCDLocator;
+	Zip64EndOfCDLocator.EndOffsetCD = EndOffsetCD;
+	Write((void*)&Zip64EndOfCDLocator, sizeof(FZip64EndOfCDLocator));
+
+	///////////////////////////
+	// Write regular end of central directory locator
+	///////////////////////////
+	FZipEndOfCDRecord ZipEndOfCDRecord;
+	Write((void*)&ZipEndOfCDRecord, sizeof(FZipEndOfCDRecord));
 
 	Flush();
 
@@ -126,6 +217,15 @@ void FZipArchiveWriter::AddFile(const FString& Filename, TConstArrayView<uint8>
 	{
 		return;
 	}
+
+	if (EnumHasAnyFlags(ZipOptions, EZipArchiveOptions::RemoveDuplicate))
+	{
+		if (Files.IndexOfByPredicate([&SearchName = Filename](const FFileEntry& Entry) { return Entry.Filename == SearchName; }) != INDEX_NONE)
+		{
+			return;
+		}
+	}
+
 	uint32 Crc = FCrc::MemCrc32(Data.GetData(), Data.Num());
 
 	// Convert the date-time to a zip file timestamp (2-second resolution).
@@ -139,36 +239,73 @@ void FZipArchiveWriter::AddFile(const FString& Filename, TConstArrayView<uint8>
 
 	uint64 FileOffset = Tell();
 
-	FFileEntry* Entry = new (Files) FFileEntry(Filename, Crc, Data.Num(), FileOffset, ZipTime);
+	TArray<uint8> TempCompressedData;
+	int32 OriginalSize = Data.Num();
+	int32 FinalDataSize = Data.Num();
+	void* WriteData = (void*)Data.GetData();
+	bool bCompressionIsValid = false;
 
-	// Local File Header
-	static const uint8 Header[] =
+	// Compress the data 
+	if (EnumHasAnyFlags(ZipOptions, EZipArchiveOptions::Deflate))
 	{
-		0x50, 0x4b, 0x03, 0x04, // Local file header signature
-		0x2d, 0x00, // Version needed to extract (MS DOS - v4.5)
-		0x00, 0x08, // General purpose bit flag (Language encoding flag = 1)
-		0x00, 0x00 // Compression method (none)
-	};
-	Write((void*)Header, sizeof(Header));
-	Write(ZipTime);
-	Write(Crc);
-
-	// Compressed and Uncompressed size - (set to 0xff as it is provided by the Zip64 block).
-	Write((uint64)0xffffffffffffffff);
+		// pre allocate a buffer for the compression result
+		int32 CompressedSize = FCompression::CompressMemoryBound(NAME_Zlib, OriginalSize, COMPRESS_NoFlags);
+		TempCompressedData.AddUninitialized(CompressedSize);
+		bool bSuccess = FCompression::CompressMemory(
+			NAME_Zlib,
+			TempCompressedData.GetData(),
+			CompressedSize,
+			WriteData,
+			OriginalSize,
+			COMPRESS_NoFlags,
+			// Flags passed into zlib's defalteInit2 function.
+			// From zlib.net/manual.html:
+			// windowBits can also be –8..–15 for raw deflate. In this case, -windowBits determines the window size. 
+			// deflate() will then  generate raw deflate data with no zlib header or trailer, and will not compute a check value.
+			-DEFAULT_ZLIB_BIT_WINDOW);
+
+		if (bSuccess)
+		{
+			WriteData = (void*)TempCompressedData.GetData();
+			bCompressionIsValid = true;
+			FinalDataSize = CompressedSize;
+		}
+	}
+
+	FFileEntry* Entry = new (Files) FFileEntry(Filename, Crc, FileOffset, ZipTime);
+	Entry->CompressedSize = FinalDataSize;
+	Entry->UnCompressedSize = OriginalSize;
+	Entry->bIsCompress = bCompressionIsValid;
 
+	///////////////////////////
+	// Local File Header
+	///////////////////////////
 	FUtf8String UTF8Filename = *Entry->Filename;
-	Write((uint16)UTF8Filename.Len());
-	Write((uint16)0x14); // Length of extra fields (Zip64 Extended Information)
+	FZipLocalHeader ZipLocalHeader;
+	ZipLocalHeader.GenPurposeBit = (uint16)EZipGenPurposeFlags::LanguageEncodingFlag;
+	ZipLocalHeader.CompMode = bCompressionIsValid ? ZIP_CM_DEFLATE : ZIP_CM_STORE;
+	ZipLocalHeader.TimeDate = ZipTime;
+	ZipLocalHeader.Crc = Crc;
+	ZipLocalHeader.CompressedSize = 0xFFFFFFFF; // Set to 0xff as it is provided by the Zip64 block
+	ZipLocalHeader.UnCompressedSize = 0xFFFFFFFF; // Set to 0xff as it is provided by the Zip64 block	
+	ZipLocalHeader.FileNameLength = UTF8Filename.Len();
+	ZipLocalHeader.ExtraFieldLength = sizeof(FZip64ExtraFieldHeader);
+	Write((void*)&ZipLocalHeader, sizeof(FZipLocalHeader));
+	// write the variable size data
 	Write((void*)GetData(UTF8Filename), UTF8Filename.Len());
 
+	///////////////////////////
 	// Zip64 Extended Information block
-	Write((uint16)0x01);		// Zip64 tag
-	Write((uint16)0x10);		// Size of this block (16 bytes)
-	Write((uint64)Data.Num());  // Uncompressed size
-	Write((uint64)Data.Num());  // Compressed size
-
-	Write((void*)Data.GetData(), Data.Num());
-
+	///////////////////////////
+	FZip64ExtraFieldHeader Zip64ExtraFieldHeader;
+	Zip64ExtraFieldHeader.UnCompressedSize = OriginalSize;
+	Zip64ExtraFieldHeader.CompressedSize = FinalDataSize;
+	Write((void*)&Zip64ExtraFieldHeader, sizeof(FZip64ExtraFieldHeader));
+
+	///////////////////////////
+	// Write the file data, either compressed or uncompressed as set above
+	///////////////////////////
+	Write((void*)WriteData, FinalDataSize);
 	Flush();
 }
 
@@ -192,4 +329,4 @@ void FZipArchiveWriter::Flush()
 	}
 }
 
-#endif
\ No newline at end of file
+#endif
diff --git a/Engine/Source/Developer/FileUtilities/Public/FileUtilities/ZipArchiveWriter.h b/Engine/Source/Developer/FileUtilities/Public/FileUtilities/ZipArchiveWriter.h
--- a/Engine/Source/Developer/FileUtilities/Public/FileUtilities/ZipArchiveWriter.h
+++ b/Engine/Source/Developer/FileUtilities/Public/FileUtilities/ZipArchiveWriter.h
@@ -10,11 +10,20 @@
 #include "HAL/Platform.h"
 #include "HAL/UnrealMemory.h"
 #include "Misc/DateTime.h"
-
+ #include "Misc/EnumClassFlags.h"
+ 
 #if WITH_ENGINE
 
 class IFileHandle;
 
+enum class EZipArchiveOptions : uint8
+{
+	None			= 0,
+	Deflate			= 1 << 0,
+	RemoveDuplicate = 1 << 1
+};
+ENUM_CLASS_FLAGS(EZipArchiveOptions);
+
 /** Helper class for generating an uncompressed zip archive file. */
 class FILEUTILITIES_API FZipArchiveWriter
 {
@@ -22,24 +31,30 @@ class FILEUTILITIES_API FZipArchiveWriter
 	{
 		FString Filename;
 		uint32 Crc32;
-		uint64 Length;
 		uint64 Offset;
 		uint32 Time;
+		uint64 CompressedSize;
+		uint64 UnCompressedSize;
+		bool bIsCompress;
 
-		FFileEntry(const FString& InFilename, uint32 InCrc32, uint64 InLength, uint64 InOffset, uint32 InTime)
+		FFileEntry(const FString& InFilename, uint32 InCrc32, uint64 InOffset, uint32 InTime)
 			: Filename(InFilename)
 			, Crc32(InCrc32)
-			, Length(InLength)
 			, Offset(InOffset)
 			, Time(InTime)
-		{}
+			, CompressedSize(0)
+			, UnCompressedSize(0)
+			, bIsCompress(false)
+		{
+		}
 	};
 
 	TArray<FFileEntry> Files;
 
 	TArray<uint8> Buffer;
 	IFileHandle* File;
-
+	EZipArchiveOptions ZipOptions;
+	
 	inline void Write(uint16 V) { Write((void*)&V, sizeof(V)); }
 	inline void Write(uint32 V) { Write((void*)&V, sizeof(V)); }
 	inline void Write(uint64 V) { Write((void*)&V, sizeof(V)); }
@@ -55,7 +70,7 @@ class FILEUTILITIES_API FZipArchiveWriter
 	void Flush();
 
 public:
-	FZipArchiveWriter(IFileHandle* InFile);
+	FZipArchiveWriter(IFileHandle* InFile, EZipArchiveOptions InZipOptions = EZipArchiveOptions::None);
 	~FZipArchiveWriter();
 
 	void AddFile(const FString& Filename, TConstArrayView<uint8> Data, const FDateTime& Timestamp);
diff --git a/Engine/Source/Developer/ShaderCompilerCommon/Private/ShaderSymbolExport.cpp b/Engine/Source/Developer/ShaderCompilerCommon/Private/ShaderSymbolExport.cpp
--- a/Engine/Source/Developer/ShaderCompilerCommon/Private/ShaderSymbolExport.cpp
+++ b/Engine/Source/Developer/ShaderCompilerCommon/Private/ShaderSymbolExport.cpp
@@ -94,8 +94,10 @@ void FShaderSymbolExport::Initialize()
 			uint32 MultiprocessId = UE::GetMultiprocessId();
 			bMultiprocessOwner = MultiprocessId == 0;
 
-			// Check if the export mode is to an uncompressed archive or loose files.
-			const bool bExportAsZip = ShouldWriteShaderSymbolsAsZip(ShaderFormat);
+			// Check if the export mode is to an uncompressed/compressed archive or loose files.
+			const EWriteShaderSymbols WriteShaderSymbolsOptions = GetWriteShaderSymbolsOptions(ShaderFormat);
+			const bool bExportAsZip = WriteShaderSymbolsOptions != EWriteShaderSymbols::Disable;
+
 			if (bSymbolsEnabled && (bExportAsZip || FParse::Param(FCommandLine::Get(), TEXT("ShaderSymbolsExportZip"))))
 			{
 				FString LeafName = CreateNameAndDeleteOld(MultiprocessId, PlatformFile, ExportPath, ZipFileBaseLeafName, ZipFileExtension);
@@ -109,7 +111,10 @@ void FShaderSymbolExport::Initialize()
 				}
 				else
 				{
-					ZipWriter = MakeUnique<FZipArchiveWriter>(OutputZipFile);
+					// only compress the first zip file. Other ones are going to be compress in the merging
+					EZipArchiveOptions ZipOptions = (WriteShaderSymbolsOptions == EWriteShaderSymbols::Compress && bMultiprocessOwner) ? EZipArchiveOptions::Deflate : EZipArchiveOptions::None;
+					ZipOptions |= EZipArchiveOptions::RemoveDuplicate;
+					ZipWriter = MakeUnique<FZipArchiveWriter>(OutputZipFile, ZipOptions);
 				}
 			}
 			
diff --git a/Engine/Source/Runtime/Engine/Private/ShaderCompiler/ShaderCompiler.cpp b/Engine/Source/Runtime/Engine/Private/ShaderCompiler/ShaderCompiler.cpp
--- a/Engine/Source/Runtime/Engine/Private/ShaderCompiler/ShaderCompiler.cpp
+++ b/Engine/Source/Runtime/Engine/Private/ShaderCompiler/ShaderCompiler.cpp
@@ -422,7 +422,8 @@ static TAutoConsoleVariable<int32> CVarShadersWriteSymbolsZip(
 	TEXT("r.Shaders.WriteSymbols.Zip"),
 	0,
 	TEXT(" 0: Export as loose files.\n")
-	TEXT(" 1: Export as an uncompressed archive.\n"),
+	TEXT(" 1: Export as an uncompressed archive.\n")
+	TEXT(" 2: Export as a compressed archive.\n"),
 	ECVF_ReadOnly);
 
 static TAutoConsoleVariable<int32> CVarShadersEnableExtraData(
diff --git a/Engine/Source/Runtime/RenderCore/Private/ShaderCore.cpp b/Engine/Source/Runtime/RenderCore/Private/ShaderCore.cpp
--- a/Engine/Source/Runtime/RenderCore/Private/ShaderCore.cpp
+++ b/Engine/Source/Runtime/RenderCore/Private/ShaderCore.cpp
@@ -819,10 +819,10 @@ bool GetShaderSymbolPathOverride(FString& OutPathOverride, FName ShaderFormat)
 	return false;
 }
 
-bool ShouldWriteShaderSymbolsAsZip(FName ShaderFormat)
+EWriteShaderSymbols GetWriteShaderSymbolsOptions(FName ShaderFormat)
 {
 	static const FShaderSymbolSettingHelper WriteSymbolsZip(TEXT("r.Shaders.WriteSymbols.Zip"));
-	return WriteSymbolsZip.IsEnabled(ShaderFormat);
+	return static_cast<EWriteShaderSymbols>(WriteSymbolsZip.GetIntForPlatform(ShaderFormat));
 }
 
 bool ShouldEnableExtraShaderData(FName ShaderFormat)
diff --git a/Engine/Source/Runtime/RenderCore/Public/ShaderCompilerCore.h b/Engine/Source/Runtime/RenderCore/Public/ShaderCompilerCore.h
--- a/Engine/Source/Runtime/RenderCore/Public/ShaderCompilerCore.h
+++ b/Engine/Source/Runtime/RenderCore/Public/ShaderCompilerCore.h
@@ -67,8 +67,15 @@ extern RENDERCORE_API bool GetShaderSymbolPathOverride(FString& OutPathOverride,
 /** Returns true if (external) shader symbols should be specific to each shader rather than be de-duplicated. */
 extern RENDERCORE_API bool ShouldAllowUniqueShaderSymbols(FName ShaderFormat);
 
+enum class EWriteShaderSymbols : uint8
+{
+	Disable		= 0,
+	UnCompress	= 1,
+	Compress	= 2
+};
+
 /** Returns true if shaders should be combined into a single zip file instead of individual files. */
-extern RENDERCORE_API bool ShouldWriteShaderSymbolsAsZip(FName ShaderFormat);
+extern RENDERCORE_API EWriteShaderSymbols GetWriteShaderSymbolsOptions(FName ShaderFormat);
 
 /** Returns true if the user wants more runtime shader data (names, extra info) */
 extern RENDERCORE_API bool ShouldEnableExtraShaderData(FName ShaderFormat);
